Annotated TemplateCategories.ini and TemplateDefs.ini in the TP_Blank template with information about the values contained within.

diff --git a/Templates/TP_Blank/Config/TemplateDefs.ini b/Templates/TP_Blank/Config/TemplateDefs.ini
--- a/Templates/TP_Blank/Config/TemplateDefs.ini
+++ b/Templates/TP_Blank/Config/TemplateDefs.ini
@@ -1,9 +1,34 @@
+; This file specifies how a template will be treated when creating a project from it.
+; It includes rules for ignoring files, replacing files and changing the contents of files.
+; This file is annotated to help with creating your own templates.
+
 [/Script/GameProjectGeneration.TemplateProjectDefs]
+
+; The localized display names for this template. Each entry specifies a Language code and Unicode text to display in the New Project dialog.
 LocalizedDisplayNames=(Language="en", Text="Blank")
+
+; The localized descriptions for this template, which is displayed when the template is selected.
 LocalizedDescriptions=(Language="en", Text="A clean empty project with no code.")
 
+; The categories that this template should appear in. The values in this list refer to the Key value given in TemplateCategories.ini.
+; Note: Names should not be quoted. You can add more categories as shown with the commented out line.
+;Categories=Visualization
+Categories=ME
+Categories=AEC
+Categories=MFG
+
+; The sort key to use for the template to bypass the default alphabetical sorting.
+; Omit this if you just want to use default sorting.
+SortKey="_1"
+
+; A list of folders to avoid copying when creating a project from this template. 
+; These should contain files only of concern to the template itself and not desired in the user's project.
 FoldersToIgnore=Media
 
+; The files to ignore copying, relative to the template root folder. 
+; %TEMPLATENAME% will be replaced by the template's folder name when resolving these. 
+; Filename comparisons are case-insensitive. 
+; For example, if this template is "TP_Blank" and the user creates "MyProject", then "TP_Blank.uproject" will be copied to "MyProject.uproject".
 FilesToIgnore="%TEMPLATENAME%.uproject"
 FilesToIgnore="%TEMPLATENAME%.png"
 FilesToIgnore="Config/TemplateDefs.ini"
@@ -16,15 +41,65 @@ FilesToIgnore="%TEMPLATENAME%.sln"
 FilesToIgnore="Manifest.json"
 FilesToIgnore="contents.txt"
 
+; NOTE: For the purposes of the below examples, this template is considered to be named "TP_Blank", and the user has created a project called "MyProject".
+
+; Folders that will be renamed when creating a project from this template.
+; %PROJECTNAME% will be replaced by the project name specified by the user during creation when resolving these rules.
+; For example,  then "Source/TP_Blank" will be copied to "Source/MyProject".
 FolderRenames=(From="Source/%TEMPLATENAME%", To="Source/%PROJECTNAME%")
 FolderRenames=(From="Source/%TEMPLATENAME%Editor", To="Source/%PROJECTNAME%Editor")
 
+; Rules for renaming files when creating a project from this template. 
+; %TEMPLATENAME_UPPERCASE% and %PROJECTNAME_UPPERCASE% are the template and project names in upper case, eg. for "TP_Blank", its upper case representation is "TP_BLANK".
+; %TEMPLATENAME_LOWERCASE% and %PROJECTNAME_LOWERCASE% are the template and project names in lower case, eg. for "TP_Blank", its lower case representation is "tp_blank".
+; A rule does not have to include any of these modifiers. 
+; Extensions is a list of files to consider for these replacements. Eg. ("cpp", "h") will try to apply a rule only to files like "SomeFile.cpp" and "SomeFile.h".
+; bCaseSensitive specifies whether this replacement should be case-sensitive.
+
+; Here are some default rules for renaming files:
+
+; This rule states that a file called "TP_BLANK.cpp" will be renamed to "MYPROJECT.cpp".
 FilenameReplacements=(Extensions=("cpp","h","ini","cs"), From="%TEMPLATENAME_UPPERCASE%", To="%PROJECTNAME_UPPERCASE%", bCaseSensitive=true)
+
+; This rule states that a file called "tp_blank.ini" will be renamed to "myproject.ini"
 FilenameReplacements=(Extensions=("cpp","h","ini","cs"), From="%TEMPLATENAME_LOWERCASE%", To="%PROJECTNAME_LOWERCASE%", bCaseSensitive=true)
+
+; And this rule catches any irregularly-cased files, so that "Tp_bLANK.ini" will be renamed to "MyProject.ini".
 FilenameReplacements=(Extensions=("cpp","h","ini","cs"), From="%TEMPLATENAME%",To="%PROJECTNAME%",bCaseSensitive=false)
 
+; Rules for renaming strings in the contents of files. 
+; The same rules as above apply, but any usages of strings that match these rules inside the file will be changed. This includes all substrings!
+; For example, "TP_Blank" will be replaced by "MyProject", but "TP_Blank_Name" will also be replaced by "MyProject_Name".
 ReplacementsInFiles=(Extensions=("cpp","h","ini","cs"), From="%TEMPLATENAME_UPPERCASE%", To="%PROJECTNAME_UPPERCASE%", bCaseSensitive=true)
 ReplacementsInFiles=(Extensions=("cpp","h","ini","cs"), From="%TEMPLATENAME_LOWERCASE%", To="%PROJECTNAME_LOWERCASE%", bCaseSensitive=true)
 ReplacementsInFiles=(Extensions=("cpp","h","ini","cs"), From="%TEMPLATENAME%", To="%PROJECTNAME%", bCaseSensitive=false)
 
-SortKey="_1"
+; The following settings are not enabled for this project, but are included for completeness:
+
+; Whether or not this template can be used to create a project. Default: true
+; If this is set to false, then the template becomes a Feature Pack that can only be referenced from other templates, or added to a project from the Content Browser.
+;bAllowProjectCreation=false
+
+; Whether or not to skip the project settings page entirely for this template. Default: false
+;bSkipProjectSettings=true
+
+; Which settings on the project settings page should be hidden for this template. Default: none
+; Allowed values are:
+; Languages - Choose between C++ and Blueprint versions.
+; HardwareTarget - Choose between Mobile and Desktop hardware targets.
+; GraphicsPreset - Choose between Maximum and Scalable graphics.
+; StarterContent - Choose whether to add starter content to the project.
+; VR - Choose whether to enable Virtual Reality support for the project.
+;HiddenSettings=(Languages, HardwareTarget)
+
+
+; The class types that this project uses.
+; This is only for display in the template description, and does not affect project creation.
+;ClassTypes=GameMode, Character, Actor
+
+; The asset types that this project uses.
+; This is only for display in the template description, and does not affect project creation.
+;AssetTypes=Skeleton, Static Mesh, Texture
+
+; Shared feature packs. The files in these packs listed in these structures marked as 'additionalfiles' will be copied on project generation.
+;SharedContentPacks=(MountName="2DSideScroller",DetailLevels=("Standard"))
diff --git a/Templates/TemplateCategories.ini b/Templates/TemplateCategories.ini
--- a/Templates/TemplateCategories.ini
+++ b/Templates/TemplateCategories.ini
@@ -1,7 +1,18 @@
-[/Script/GameProjectGeneration.TemplateCategories]
+; This file lists all the template categories that are exposed from this folder.
+; Add template categories here to have them show up in the Project Browser.
+; Additional TemplateDefs.ini files can be created in your plugin root folder, at the same level as a Templates/ folder.
+; Categories will contain project templates that include the category Key in their Categories property in their Config/TemplateDefs.ini.
 
-Categories=(Key="Game", LocalizedDisplayNames=((Language="en", Text="Game")), LocalizedDescriptions=((Language="en", Text="Making fancy games with the running and the questing and looting with the juicy inventories and UIs.")), Icon="GameProjectDialog.CodeIcon", Image="GameProjectDialog.CodeImage", IsMajorCategory=true)
+[/Script/GameProjectGeneration.TemplateCategories]
 
-Categories=(Key="Virtual Production", LocalizedDisplayNames=((Language="en", Text="Virtual Production")), LocalizedDescriptions=((Language="en", Text="Movies with all the power of the Unreal and doodads like nDisplay and the compositions.")), Icon="GameProjectDialog.BlueprintIcon", Image="GameProjectDialog.BlueprintImage", IsMajorCategory=false)
+; Each "Categories" entry adds a separate category with these fields:
+; "Key": The name used to reference the category from templates. Should be short and easy.
+; "LocalizedDisplayNames": A list of pairs of "Language" codes and "Text" to display as the name of the category in the UI (supports Unicode).
+; "LocalizedDescriptions": As with display names, this is the description shown next to each category.
+; "Icon": A relative reference to the icon to display for the category - should be at minimum 128x128 to avoid being blurry.
+; "IsMajorCategory": Is this a major, top-level category? Major categories are displayed as full rows, eg. the "Game" category.
 
-Categories=(Key="Visualization", LocalizedDisplayNames=((Language="en", Text="Visualization")), LocalizedDescriptions=((Language="en", Text="Pretty pretty buildings and cars and all that.")), Icon="GameProjectDialog.BlueprintIcon", Image="GameProjectDialog.BlueprintImage", IsMajorCategory=false)
\ No newline at end of file
+Categories=(Key="Games", LocalizedDisplayNames=((Language="en", Text="Games")), LocalizedDescriptions=((Language="en", Text="Starting off points giving you a few key classes, levels and examples to begin your game.")), Icon="GameProjectDialog.BlueprintIcon", IsMajorCategory=true)
+Categories=(Key="ME", LocalizedDisplayNames=((Language="en", Text="Media & Entertainment")), LocalizedDescriptions=((Language="en", Text="need text need text need text")), Icon="GameProjectDialog.BlueprintIcon", IsMajorCategory=true)
+Categories=(Key="AEC", LocalizedDisplayNames=((Language="en", Text="Architecture, Engineering & Construction")), LocalizedDescriptions=((Language="en", Text="need text need text need text")), Icon="GameProjectDialog.BlueprintIcon", IsMajorCategory=true)
+Categories=(Key="MFG", LocalizedDisplayNames=((Language="en", Text="Manufacturing, Automotive & Product Design")), LocalizedDescriptions=((Language="en", Text="need text need text need text")),  Icon="GameProjectDialog.BlueprintIcon", IsMajorCategory=true)
\ No newline at end of file
