Blank templates now include the Geometry module, which apparently allows Starter Content to show up in the content browser.

diff --git a/Templates/TP_Blank/Config/TemplateDefs.ini b/Templates/TP_Blank/Config/TemplateDefs.ini
--- a/Templates/TP_Blank/Config/TemplateDefs.ini
+++ b/Templates/TP_Blank/Config/TemplateDefs.ini
@@ -12,10 +12,10 @@ LocalizedDescriptions=(Language="en", Text="A clean empty project with no code."
 
 ; The categories that this template should appear in. The values in this list refer to the Key value given in TemplateCategories.ini.
 ; Note: Names should not be quoted. You can add more categories as shown with the commented out line.
-;Categories=Visualization
-Categories=ME
-Categories=AEC
-Categories=MFG
+;Categories=Game
+;Categories=ME
+;Categories=AEC
+;Categories=MFG
 
 ; The sort key to use for the template to bypass the default alphabetical sorting.
 ; Omit this if you just want to use default sorting.
@@ -90,7 +90,8 @@ ReplacementsInFiles=(Extensions=("cpp","h","ini","cs"), From="%TEMPLATENAME%", T
 ; GraphicsPreset - Choose between Maximum and Scalable graphics.
 ; StarterContent - Choose whether to add starter content to the project.
 ; VR - Choose whether to enable Virtual Reality support for the project.
-;HiddenSettings=(Languages, HardwareTarget)
+;HiddenSettings=Languages
+;HiddenSettings=HardwareTarget
 
 
 ; The class types that this project uses.
diff --git a/Templates/TP_Blank/manifest.json b/Templates/TP_Blank/manifest.json
--- a/Templates/TP_Blank/manifest.json
+++ b/Templates/TP_Blank/manifest.json
@@ -25,5 +25,15 @@
 		"TP_Blank_Preview.png"
 	],
 	"SortKey": "_1",
-	"AdditionalFeaturePacks": []
+	"AdditionalFeaturePacks":
+	[
+		{
+			"MountName": "Geometry",
+			"DetailLevels":
+			[
+				"Standard",
+				"High"
+			]
+		}
+	]
 }
diff --git a/Templates/TP_BlankBP/Config/TemplateDefs.ini b/Templates/TP_BlankBP/Config/TemplateDefs.ini
--- a/Templates/TP_BlankBP/Config/TemplateDefs.ini
+++ b/Templates/TP_BlankBP/Config/TemplateDefs.ini
@@ -2,6 +2,13 @@
 LocalizedDisplayNames=(Language="en", Text="Blank")
 LocalizedDescriptions=(Language="en", Text="A clean empty project with no code.")
 
+SortKey="_1"
+
+;Categories=Game
+;Categories=ME
+;Categories=AEC
+;Categories=MFG
+
 FoldersToIgnore=Media
 
 FilesToIgnore="%TEMPLATENAME%.uproject"
@@ -27,4 +34,4 @@ ReplacementsInFiles=(Extensions=("cpp","h","ini","cs"), From="%TEMPLATENAME_UPPE
 ReplacementsInFiles=(Extensions=("cpp","h","ini","cs"), From="%TEMPLATENAME_LOWERCASE%", To="%PROJECTNAME_LOWERCASE%", bCaseSensitive=true)
 ReplacementsInFiles=(Extensions=("cpp","h","ini","cs"), From="%TEMPLATENAME%", To="%PROJECTNAME%", bCaseSensitive=false)
 
-SortKey="_1"
+SharedContentPacks=(MountName="Geometry",DetailLevels=("Standard","High"))
\ No newline at end of file
diff --git a/Templates/TP_BlankBP/manifest.json b/Templates/TP_BlankBP/manifest.json
--- a/Templates/TP_BlankBP/manifest.json
+++ b/Templates/TP_BlankBP/manifest.json
@@ -25,5 +25,15 @@
 		"TP_BlankBP_Preview.png"
 	],
 	"SortKey": "_1",
-	"AdditionalFeaturePacks": []
+	"AdditionalFeaturePacks":
+	[
+		{
+			"MountName": "Geometry",
+			"DetailLevels":
+			[
+				"Standard",
+				"High"
+			]
+		}
+	]
 }
