Categorization of various templates to work with the new Project Creation Wizard

TempDefs.ini for each template now contains a predefined category.
TemplateCategories.ini defines categories for various markets (Games, M&E, AEC, MFG) which was agreed with Marketing.

This sets the framework to start creating various market specific templates which will be delivered in subsequent submissions.

diff --git a/Templates/FP_FirstPerson/Config/TemplateDefs.ini b/Templates/FP_FirstPerson/Config/TemplateDefs.ini
--- a/Templates/FP_FirstPerson/Config/TemplateDefs.ini
+++ b/Templates/FP_FirstPerson/Config/TemplateDefs.ini
@@ -1,5 +1,8 @@
 ﻿// This isnt actually a real template, but we need this here to have the shared resources mounted
 [/Script/GameProjectGeneration.TemplateProjectDefs]
+
+Categories=Games
+
 SharedContentPacks=(MountName="FirstPerson",DetailLevels=("Standard"))
 SharedContentPacks=(MountName="Geometry",DetailLevels=("Standard","High"))
-bAllowProjectCreation=false
\ No newline at end of file
+bAllowProjectCreation=false
diff --git a/Templates/FP_FirstPersonBP/Config/TemplateDefs.ini b/Templates/FP_FirstPersonBP/Config/TemplateDefs.ini
--- a/Templates/FP_FirstPersonBP/Config/TemplateDefs.ini
+++ b/Templates/FP_FirstPersonBP/Config/TemplateDefs.ini
@@ -1,5 +1,8 @@
 ﻿// This isnt actually a real template, but we need this here to have the shared resources mounted
 [/Script/GameProjectGeneration.TemplateProjectDefs]
+
+Categories=Games
+
 SharedContentPacks=(MountName="FirstPerson",DetailLevels=("Standard"))
 SharedContentPacks=(MountName="Geometry",DetailLevels=("Standard","High"))
 bAllowProjectCreation=false
diff --git a/Templates/TP_2DSideScroller/Config/TemplateDefs.ini b/Templates/TP_2DSideScroller/Config/TemplateDefs.ini
--- a/Templates/TP_2DSideScroller/Config/TemplateDefs.ini
+++ b/Templates/TP_2DSideScroller/Config/TemplateDefs.ini
@@ -1,4 +1,7 @@
-[/Script/GameProjectGeneration.TemplateProjectDefs]
+﻿[/Script/GameProjectGeneration.TemplateProjectDefs]
+
+Categories=Games
+
 LocalizedDisplayNames=(Language="en",Text="2D Side Scroller")
 LocalizedDescriptions=(Language="en",Text="This is a paper 2D version of the Side Scroller template. It features a 2D sprite character that uses texture based animation where the camera is positioned at the avatar's side. Character movement, which can be controlled with either keyboard, controller or virtual joystick on a touch device, is limited to two dimensions: forward, backward and jumping. A number of ledges are featured in the level and the character can jump between them.")
 
diff --git a/Templates/TP_2DSideScrollerBP/Config/TemplateDefs.ini b/Templates/TP_2DSideScrollerBP/Config/TemplateDefs.ini
--- a/Templates/TP_2DSideScrollerBP/Config/TemplateDefs.ini
+++ b/Templates/TP_2DSideScrollerBP/Config/TemplateDefs.ini
@@ -1,4 +1,7 @@
 ﻿[/Script/GameProjectGeneration.TemplateProjectDefs]
+
+Categories=Games
+
 LocalizedDisplayNames=(Language="en",Text="2D Side Scroller")
 LocalizedDescriptions=(Language="en",Text="This is a paper 2D version of the Side Scroller template. It features a 2D sprite character that uses texture based animation where the camera is positioned at the avatar's side. Character movement, which can be controlled with either keyboard, controller or virtual joystick on a touch device, is limited to two dimensions: forward, backward and jumping. A number of ledges are featured in the level and the character can jump between them.")
 
diff --git a/Templates/TP_Blank/Config/TemplateDefs.ini b/Templates/TP_Blank/Config/TemplateDefs.ini
--- a/Templates/TP_Blank/Config/TemplateDefs.ini
+++ b/Templates/TP_Blank/Config/TemplateDefs.ini
@@ -4,21 +4,17 @@
 
 [/Script/GameProjectGeneration.TemplateProjectDefs]
 
+; The categories that this template should appear in. The values in this list refer to the Key value given in TemplateCategories.ini.
+Categories=Games
+Categories=ME
+
 ; The localized display names for this template. Each entry specifies a Language code and Unicode text to display in the New Project dialog.
 LocalizedDisplayNames=(Language="en", Text="Blank")
 
 ; The localized descriptions for this template, which is displayed when the template is selected.
 LocalizedDescriptions=(Language="en", Text="A clean empty project with no code.")
 
-; The categories that this template should appear in. The values in this list refer to the Key value given in TemplateCategories.ini.
-; Note: Names should not be quoted. You can add more categories as shown with the commented out line.
-;Categories=Game
-;Categories=ME
-;Categories=AEC
-;Categories=MFG
-
 ; The sort key to use for the template to bypass the default alphabetical sorting.
-; Omit this if you just want to use default sorting.
 SortKey="_1"
 
 ; A list of folders to avoid copying when creating a project from this template. 
@@ -80,19 +76,20 @@ ReplacementsInFiles=(Extensions=("cpp","h","ini","cs"), From="%TEMPLATENAME%", T
 ; If this is set to false, then the template becomes a Feature Pack that can only be referenced from other templates, or added to a project from the Content Browser.
 ;bAllowProjectCreation=false
 
-; Whether or not to skip the project settings page entirely for this template. Default: false
-;bSkipProjectSettings=true
-
 ; Which settings on the project settings page should be hidden for this template. Default: none
 ; Allowed values are:
+; All - Hides the settings page entirely
 ; Languages - Choose between C++ and Blueprint versions.
 ; HardwareTarget - Choose between Mobile and Desktop hardware targets.
 ; GraphicsPreset - Choose between Maximum and Scalable graphics.
 ; StarterContent - Choose whether to add starter content to the project.
 ; VR - Choose whether to enable Virtual Reality support for the project.
+;HiddenSettings=All
 ;HiddenSettings=Languages
 ;HiddenSettings=HardwareTarget
-
+;HiddenSettings=GraphicsPreset 
+;HiddenSettings=StarterContent 
+;HiddenSettings=VR 
 
 ; The class types that this project uses.
 ; This is only for display in the template description, and does not affect project creation.
diff --git a/Templates/TP_BlankBP/Config/TemplateDefs.ini b/Templates/TP_BlankBP/Config/TemplateDefs.ini
--- a/Templates/TP_BlankBP/Config/TemplateDefs.ini
+++ b/Templates/TP_BlankBP/Config/TemplateDefs.ini
@@ -1,13 +1,12 @@
 [/Script/GameProjectGeneration.TemplateProjectDefs]
-LocalizedDisplayNames=(Language="en", Text="Blank")
-LocalizedDescriptions=(Language="en", Text="A clean empty project with no code.")
 
 SortKey="_1"
 
-;Categories=Game
-;Categories=ME
-;Categories=AEC
-;Categories=MFG
+Categories=Games
+Categories=ME
+
+LocalizedDisplayNames=(Language="en", Text="Blank")
+LocalizedDescriptions=(Language="en", Text="A clean empty project with no code.")
 
 FoldersToIgnore=Media
 
@@ -34,4 +33,6 @@ ReplacementsInFiles=(Extensions=("cpp","h","ini","cs"), From="%TEMPLATENAME_UPPE
 ReplacementsInFiles=(Extensions=("cpp","h","ini","cs"), From="%TEMPLATENAME_LOWERCASE%", To="%PROJECTNAME_LOWERCASE%", bCaseSensitive=true)
 ReplacementsInFiles=(Extensions=("cpp","h","ini","cs"), From="%TEMPLATENAME%", To="%PROJECTNAME%", bCaseSensitive=false)
 
-SharedContentPacks=(MountName="Geometry",DetailLevels=("Standard","High"))
\ No newline at end of file
+SharedContentPacks=(MountName="Geometry",DetailLevels=("Standard","High"))
+
+
diff --git a/Templates/TP_FirstPerson/Config/TemplateDefs.ini b/Templates/TP_FirstPerson/Config/TemplateDefs.ini
--- a/Templates/TP_FirstPerson/Config/TemplateDefs.ini
+++ b/Templates/TP_FirstPerson/Config/TemplateDefs.ini
@@ -1,4 +1,7 @@
-[/Script/GameProjectGeneration.TemplateProjectDefs]
+﻿[/Script/GameProjectGeneration.TemplateProjectDefs]
+
+Categories=Games
+
 LocalizedDisplayNames=(Language="en",Text="First Person")
 LocalizedDescriptions=(Language="en",Text="The first person template features a player character (represented by a pair of arms) which is viewed from first person perspective. The character can be moved around the level using keyboard, controller or virtual joystick on a touch device. Additionally the player can look around using mouse, controller or virtual joystick on a touch device. The character is also equipped with a gun that, using mouse,  controller or virtual joystick on a touch device will fire a simple sphere projectile that will affect some physics objects in the level, whilst bounce off the arena walls.")
 
diff --git a/Templates/TP_FirstPersonBP/Config/TemplateDefs.ini b/Templates/TP_FirstPersonBP/Config/TemplateDefs.ini
--- a/Templates/TP_FirstPersonBP/Config/TemplateDefs.ini
+++ b/Templates/TP_FirstPersonBP/Config/TemplateDefs.ini
@@ -1,4 +1,7 @@
 ﻿[/Script/GameProjectGeneration.TemplateProjectDefs]
+
+Categories=Games
+
 LocalizedDisplayNames=(Language="en",Text="First Person")
 LocalizedDescriptions=(Language="en",Text="The first person template features a player character (represented by a pair of arms) which is viewed from first person perspective. The character can be moved around the level using keyboard, controller or virtual joystick on a touch device. Additionally the player can look around using mouse, controller or virtual joystick on a touch device. The character is also equipped with a gun that, using mouse,  controller or virtual joystick on a touch device will fire a simple sphere projectile that will affect some physics objects in the level, whilst bounce off the arena walls.")
 
diff --git a/Templates/TP_Flying/Config/TemplateDefs.ini b/Templates/TP_Flying/Config/TemplateDefs.ini
--- a/Templates/TP_Flying/Config/TemplateDefs.ini
+++ b/Templates/TP_Flying/Config/TemplateDefs.ini
@@ -1,4 +1,7 @@
-[/Script/GameProjectGeneration.TemplateProjectDefs]
+﻿[/Script/GameProjectGeneration.TemplateProjectDefs]
+
+Categories=Games
+
 LocalizedDisplayNames=(Language="en",Text="Flying")
 LocalizedDescriptions=(Language="en",Text="The flying template features a spaceship with the camera positioned behind the ship represented by a static mesh. The keyboard, controller or virtual joystick controls the pitch and yaw of the ship as well as providing a means to allow speed control of the ships forward movement. The level is populated with several clusters of boxes that block the movement of the ship.")
 
diff --git a/Templates/TP_FlyingBP/Config/TemplateDefs.ini b/Templates/TP_FlyingBP/Config/TemplateDefs.ini
--- a/Templates/TP_FlyingBP/Config/TemplateDefs.ini
+++ b/Templates/TP_FlyingBP/Config/TemplateDefs.ini
@@ -1,4 +1,7 @@
 ﻿[/Script/GameProjectGeneration.TemplateProjectDefs]
+
+Categories=Games
+
 LocalizedDisplayNames=(Language="en",Text="Flying")
 LocalizedDescriptions=(Language="en",Text="The flying template features a spaceship with the camera positioned behind the ship represented by a static mesh. The keyboard, controller or virtual joystick controls the pitch and yaw of the ship as well as providing a means to allow speed control of the ships forward movement. The level is populated with several clusters of boxes that block the movement of the ship.")
 
diff --git a/Templates/TP_HandheldARBP/Config/TemplateDefs.ini b/Templates/TP_HandheldARBP/Config/TemplateDefs.ini
--- a/Templates/TP_HandheldARBP/Config/TemplateDefs.ini
+++ b/Templates/TP_HandheldARBP/Config/TemplateDefs.ini
@@ -1,4 +1,7 @@
 ﻿[/Script/GameProjectGeneration.TemplateProjectDefs]
+
+Categories=Games
+
 LocalizedDisplayNames=(Language="en",Text="Handheld AR")
 LocalizedDescriptions=(Language="en",Text="A starting point for building Augmented Reality applications for Android and iOS devices. Includes runtime logic for toggling AR mode on and off, debug information regarding plane detection, example code for hit detection and handling light estimation.")
 
diff --git a/Templates/TP_Puzzle/Config/TemplateDefs.ini b/Templates/TP_Puzzle/Config/TemplateDefs.ini
--- a/Templates/TP_Puzzle/Config/TemplateDefs.ini
+++ b/Templates/TP_Puzzle/Config/TemplateDefs.ini
@@ -1,4 +1,7 @@
-[/Script/GameProjectGeneration.TemplateProjectDefs]
+﻿[/Script/GameProjectGeneration.TemplateProjectDefs]
+
+Categories=Games
+
 LocalizedDisplayNames=(Language="en",Text="Puzzle")
 LocalizedDescriptions=(Language="en",Text="The puzzle template features a simple grid of blocks. The camera is postioned overhead of the blocks. The user can interact with these blocks by clicking using either the mouse or touch screen. When a block is clicked its color changes and a score, which shown using a text component within the level is increased.")
 
diff --git a/Templates/TP_PuzzleBP/Config/TemplateDefs.ini b/Templates/TP_PuzzleBP/Config/TemplateDefs.ini
--- a/Templates/TP_PuzzleBP/Config/TemplateDefs.ini
+++ b/Templates/TP_PuzzleBP/Config/TemplateDefs.ini
@@ -1,4 +1,7 @@
 ﻿[/Script/GameProjectGeneration.TemplateProjectDefs]
+
+Categories=Games
+
 LocalizedDisplayNames=(Language="en",Text="Puzzle")
 LocalizedDescriptions=(Language="en",Text="The puzzle template features a simple grid of blocks. The camera is postioned overhead of the blocks. The user can interact with these blocks by clicking using either the mouse or touch screen. When a block is clicked its color changes and a score, which shown using a text component within the level is increased.")
 
diff --git a/Templates/TP_Rolling/Config/TemplateDefs.ini b/Templates/TP_Rolling/Config/TemplateDefs.ini
--- a/Templates/TP_Rolling/Config/TemplateDefs.ini
+++ b/Templates/TP_Rolling/Config/TemplateDefs.ini
@@ -1,4 +1,7 @@
-[/Script/GameProjectGeneration.TemplateProjectDefs]
+﻿[/Script/GameProjectGeneration.TemplateProjectDefs]
+
+Categories=Games
+
 LocalizedDisplayNames=(Language="en",Text="Rolling")
 LocalizedDescriptions=(Language="en",Text="The rolling template features a physics-based rolling ball which can be rolled in any direction using either a keyboard, controller or virtual joystick on touch devices. The camera is positioned above and behind the ball and follows the movement of the ball. Within the level are some ramps, several physics cubes in the level that can be knocked over with the ball as well as some cubes which cannot.")
 
diff --git a/Templates/TP_RollingBP/Config/TemplateDefs.ini b/Templates/TP_RollingBP/Config/TemplateDefs.ini
--- a/Templates/TP_RollingBP/Config/TemplateDefs.ini
+++ b/Templates/TP_RollingBP/Config/TemplateDefs.ini
@@ -1,4 +1,7 @@
 ﻿[/Script/GameProjectGeneration.TemplateProjectDefs]
+
+Categories=Games
+
 LocalizedDisplayNames=(Language="en",Text="Rolling")
 LocalizedDescriptions=(Language="en",Text="The rolling template features a physics-based rolling ball which can be rolled in any direction using either a keyboard, controller or virtual joystick on touch devices. The camera is positioned above and behind the ball and follows the movement of the ball. Within the level are some ramps, several physics cubes in the level that can be knocked over with the ball as well as some cubes which cannot.")
 
diff --git a/Templates/TP_SideScroller/Config/TemplateDefs.ini b/Templates/TP_SideScroller/Config/TemplateDefs.ini
--- a/Templates/TP_SideScroller/Config/TemplateDefs.ini
+++ b/Templates/TP_SideScroller/Config/TemplateDefs.ini
@@ -1,4 +1,7 @@
-[/Script/GameProjectGeneration.TemplateProjectDefs]
+﻿[/Script/GameProjectGeneration.TemplateProjectDefs]
+
+Categories=Games
+
 LocalizedDisplayNames=(Language="en",Text="Side Scroller")
 LocalizedDescriptions=(Language="en",Text="This side scroller template features an animated skeletal mesh character with a camera positioned at the character's side. Movement, which can be controlled with either keyboard, controller or virtual joystick on a touch device is limited to two dimensions: forward, backward and jumping. A number of ledges are featured in the level and the character can jump between them.")
 
diff --git a/Templates/TP_SideScrollerBP/Config/TemplateDefs.ini b/Templates/TP_SideScrollerBP/Config/TemplateDefs.ini
--- a/Templates/TP_SideScrollerBP/Config/TemplateDefs.ini
+++ b/Templates/TP_SideScrollerBP/Config/TemplateDefs.ini
@@ -1,4 +1,7 @@
 ﻿[/Script/GameProjectGeneration.TemplateProjectDefs]
+
+Categories=Games
+
 LocalizedDisplayNames=(Language="en",Text="Side Scroller")
 LocalizedDescriptions=(Language="en",Text="This side scroller template features an animated skeletal mesh character with a camera positioned at the character's side. Movement, which can be controlled with either keyboard, controller or virtual joystick on a touch device is limited to two dimensions: forward, backward and jumping. A number of ledges are featured in the level and the character can jump between them.")
 
diff --git a/Templates/TP_ThirdPerson/Config/TemplateDefs.ini b/Templates/TP_ThirdPerson/Config/TemplateDefs.ini
--- a/Templates/TP_ThirdPerson/Config/TemplateDefs.ini
+++ b/Templates/TP_ThirdPerson/Config/TemplateDefs.ini
@@ -1,4 +1,8 @@
-[/Script/GameProjectGeneration.TemplateProjectDefs]
+﻿[/Script/GameProjectGeneration.TemplateProjectDefs]
+
+Categories=Games
+
+
 LocalizedDisplayNames=(Language="en",Text="Third Person")
 LocalizedDescriptions=(Language="en",Text="The third person template features a playable character where the camera is positioned behind and slightly above the character. As the character moves using mouse, keyboard, controller or virtual joystick on a touch device, the camera follows the character, observing from an over-the-shoulder point of view. This perspective emphasizes the main character and is popular in action and adventure games. The level contains several objects that the player can jump upon.")
 
diff --git a/Templates/TP_ThirdPersonBP/Config/TemplateDefs.ini b/Templates/TP_ThirdPersonBP/Config/TemplateDefs.ini
--- a/Templates/TP_ThirdPersonBP/Config/TemplateDefs.ini
+++ b/Templates/TP_ThirdPersonBP/Config/TemplateDefs.ini
@@ -1,4 +1,7 @@
 ﻿[/Script/GameProjectGeneration.TemplateProjectDefs]
+
+Categories=Games
+
 LocalizedDisplayNames=(Language="en",Text="Third Person")
 LocalizedDescriptions=(Language="en",Text="The third person template features a playable character where the camera is positioned behind and slightly above the character. As the character moves using mouse, keyboard, controller or virtual joystick on a touch device, the camera follows the character, observing from an over-the-shoulder point of view. This perspective emphasizes the main character and is popular in action and adventure games. The level contains several objects that the player can jump upon.")
 
diff --git a/Templates/TP_TopDown/Config/TemplateDefs.ini b/Templates/TP_TopDown/Config/TemplateDefs.ini
--- a/Templates/TP_TopDown/Config/TemplateDefs.ini
+++ b/Templates/TP_TopDown/Config/TemplateDefs.ini
@@ -1,4 +1,8 @@
-[/Script/GameProjectGeneration.TemplateProjectDefs]
+﻿[/Script/GameProjectGeneration.TemplateProjectDefs]
+
+Categories=Games
+
+
 LocalizedDisplayNames=(Language="en",Text="Top Down")
 LocalizedDescriptions=(Language="en",Text="The top down template features a playable character where the camera is positioned behind and above the avatar at greater distance than in third person perspective. The player is controlled by using the mouse or touchscreen to click on the required destination and uses the navigation system to facilitate the character movement. This perspective is often used in action role playing games.")
 
diff --git a/Templates/TP_TopDownBP/Config/TemplateDefs.ini b/Templates/TP_TopDownBP/Config/TemplateDefs.ini
--- a/Templates/TP_TopDownBP/Config/TemplateDefs.ini
+++ b/Templates/TP_TopDownBP/Config/TemplateDefs.ini
@@ -1,4 +1,8 @@
 ﻿[/Script/GameProjectGeneration.TemplateProjectDefs]
+
+Categories=Games
+
+
 LocalizedDisplayNames=(Language="en",Text="Top Down")
 LocalizedDescriptions=(Language="en",Text="The top down template features a playable character where the camera is positioned behind and above the avatar at greater distance than in third person perspective. The player is controlled by using the mouse or touchscreen to click on the required destination and uses the navigation system to facilitate the character movement. This perspective is often used in action role playing games.")
 
diff --git a/Templates/TP_TwinStick/Config/TemplateDefs.ini b/Templates/TP_TwinStick/Config/TemplateDefs.ini
--- a/Templates/TP_TwinStick/Config/TemplateDefs.ini
+++ b/Templates/TP_TwinStick/Config/TemplateDefs.ini
@@ -1,4 +1,8 @@
-[/Script/GameProjectGeneration.TemplateProjectDefs]
+﻿[/Script/GameProjectGeneration.TemplateProjectDefs]
+
+Categories=Games
+
+
 LocalizedDisplayNames=(Language="en",Text="Twin Stick Shooter")
 LocalizedDescriptions=(Language="en",Text="The twin stick shooter template features a playable character where the camera is positioned almost directly above the character. Control of the character is with either keyboard, controller or virtual joystick on a touch device. Movement and firing, which are both on the X/Y axis are controlled independantly allowing for firing in a different direction than which the character is facing. The level contains several cubes that are affected by the projectiles and several that are not.")
 
diff --git a/Templates/TP_TwinStickBP/Config/TemplateDefs.ini b/Templates/TP_TwinStickBP/Config/TemplateDefs.ini
--- a/Templates/TP_TwinStickBP/Config/TemplateDefs.ini
+++ b/Templates/TP_TwinStickBP/Config/TemplateDefs.ini
@@ -1,4 +1,7 @@
 ﻿[/Script/GameProjectGeneration.TemplateProjectDefs]
+
+Categories=Games
+
 LocalizedDisplayNames=(Language="en",Text="Twin Stick Shooter")
 LocalizedDescriptions=(Language="en",Text="The twin stick shooter template features a playable character where the camera is positioned almost directly above the character. Control of the character is with either keyboard, controller or virtual joystick on a touch device. Movement and firing, which are both on the X/Y axis are controlled independantly allowing for firing in a different direction than which the character is facing. The level contains several cubes that are affected by the projectiles and several that are not.")
 
diff --git a/Templates/TP_Vehicle/Config/TemplateDefs.ini b/Templates/TP_Vehicle/Config/TemplateDefs.ini
--- a/Templates/TP_Vehicle/Config/TemplateDefs.ini
+++ b/Templates/TP_Vehicle/Config/TemplateDefs.ini
@@ -1,4 +1,7 @@
-[/Script/GameProjectGeneration.TemplateProjectDefs]
+﻿[/Script/GameProjectGeneration.TemplateProjectDefs]
+
+Categories=Games
+
 LocalizedDisplayNames=(Language="en",Text="Vehicle")
 LocalizedDescriptions=(Language="en",Text="The vehicle template features a simple physics-driven vehicle. Vehicle movement can be controlled with either keyboard, controller or virtual joystick on a touch device. The template features 2 camera views - On is positioned above and behind the vehicle, whilst the other is positioned inside. When using the inside view the mouse or touchscreen or a HMD device. There is also a HUD both in and out of the car that shows the current gear and speed. The level features a smooth ramp, a slightly irregluar one, some objects that can be knocked over and a post which cannot.")
 
diff --git a/Templates/TP_VehicleAdv/Config/TemplateDefs.ini b/Templates/TP_VehicleAdv/Config/TemplateDefs.ini
--- a/Templates/TP_VehicleAdv/Config/TemplateDefs.ini
+++ b/Templates/TP_VehicleAdv/Config/TemplateDefs.ini
@@ -1,4 +1,7 @@
-[/Script/GameProjectGeneration.TemplateProjectDefs]
+﻿[/Script/GameProjectGeneration.TemplateProjectDefs]
+
+Categories=Games
+
 LocalizedDisplayNames=(Language="en",Text="VehicleAdvanced")
 LocalizedDescriptions=(Language="en",Text="The vehicle advanced template features a more complex vehicle than the regular vehicle template. Vehicle movement can be controlled with either keyboard, controller or virtual joystick on a touch device. The template features 2 camera views - One is positioned above and behind the vehicle, whilst the other is positioned inside. When using the inside view you can look around using the mouse or touchscreen or a HMD device. The vehicle itself features working double wishbone suspension. There is also an engine sound in this template. There is also a HUD both in and out of the car that shows the current gear and speed. The level contains a landscape with rough terrain as well as a simple track that features bumps, ramps and loop the loops.")
 
diff --git a/Templates/TP_VehicleAdvBP/Config/TemplateDefs.ini b/Templates/TP_VehicleAdvBP/Config/TemplateDefs.ini
--- a/Templates/TP_VehicleAdvBP/Config/TemplateDefs.ini
+++ b/Templates/TP_VehicleAdvBP/Config/TemplateDefs.ini
@@ -1,4 +1,7 @@
 ﻿[/Script/GameProjectGeneration.TemplateProjectDefs]
+
+Categories=Games
+
 LocalizedDisplayNames=(Language="en",Text="VehicleAdvanced")
 LocalizedDescriptions=(Language="en",Text="The vehicle advanced template features a more complex vehicle than the regular vehicle template. Vehicle movement can be controlled with either keyboard, controller or virtual joystick on a touch device. The template features 2 camera views - One is positioned above and behind the vehicle, whilst the other is positioned inside. When using the inside view you can look around using the mouse or touchscreen or a HMD device. The vehicle itself features working double wishbone suspension. There is also an engine sound in this template. There is also a HUD both in and out of the car that shows the current gear and speed. The level contains a landscape with rough terrain as well as a simple track that features bumps, ramps and loop the loops.")
 
diff --git a/Templates/TP_VehicleBP/Config/TemplateDefs.ini b/Templates/TP_VehicleBP/Config/TemplateDefs.ini
--- a/Templates/TP_VehicleBP/Config/TemplateDefs.ini
+++ b/Templates/TP_VehicleBP/Config/TemplateDefs.ini
@@ -1,4 +1,7 @@
 ﻿[/Script/GameProjectGeneration.TemplateProjectDefs]
+
+Categories=Games
+
 LocalizedDisplayNames=(Language="en",Text="Vehicle")
 LocalizedDescriptions=(Language="en",Text="The vehicle template features a simple physics-driven vehicle. Vehicle movement can be controlled with either keyboard, controller or virtual joystick on a touch device. The template features 2 camera views - One is positioned above and behind the vehicle, whilst the other is positioned inside. When using the inside view the mouse or touchscreen or a HMD device. There is also a HUD both in and out of the car that shows the current gear and speed. The level features a smooth ramp, a slightly irregluar one, some objects that can be knocked over and a post which cannot.")
 
diff --git a/Templates/TP_VirtualRealityBP/Config/TemplateDefs.ini b/Templates/TP_VirtualRealityBP/Config/TemplateDefs.ini
--- a/Templates/TP_VirtualRealityBP/Config/TemplateDefs.ini
+++ b/Templates/TP_VirtualRealityBP/Config/TemplateDefs.ini
@@ -1,4 +1,7 @@
 ﻿[/Script/GameProjectGeneration.TemplateProjectDefs]
+
+Categories=Games
+
 LocalizedDisplayNames=(Language="en",Text="Virtual Reality")
 LocalizedDescriptions=(Language="en",Text="Blueprint Virtual Reality Template for Desktop and Playstation VR. The template features two modes of teleport locomotion with separate pawns and maps for each. The MotionControllerMap lets players teleport around using motion controllers. The HMDLocomotionMap lets players teleport using either a Gamepad or point-and-teleport using just the HeadMountedDisplay.")
 
diff --git a/Templates/TP_nDisplayBP/Config/TemplateDefs.ini b/Templates/TP_nDisplayBP/Config/TemplateDefs.ini
--- a/Templates/TP_nDisplayBP/Config/TemplateDefs.ini
+++ b/Templates/TP_nDisplayBP/Config/TemplateDefs.ini
@@ -1,4 +1,7 @@
 ﻿[/Script/GameProjectGeneration.TemplateProjectDefs]
+
+Categories=ME
+
 LocalizedDisplayNames=(Language="en",Text="nDisplay")
 LocalizedDescriptions=(Language="en",Text="This project template provides multi-display functionality using a cluster of PCs. Use as a starting point for stereoscropic/monoscopic Caves and PowerWall setups.")
 
@@ -34,4 +37,19 @@ FilenameReplacements=(Extensions=("cpp","h","ini","cs"),From="%TEMPLATENAME%",To
 
 ReplacementsInFiles=(Extensions=("cpp","h","ini","cs"),From="%TEMPLATENAME_UPPERCASE%",To="%PROJECTNAME_UPPERCASE%",bCaseSensitive=true)
 ReplacementsInFiles=(Extensions=("cpp","h","ini","cs"),From="%TEMPLATENAME_LOWERCASE%",To="%PROJECTNAME_LOWERCASE%",bCaseSensitive=true)
-ReplacementsInFiles=(Extensions=("cpp","h","ini","cs"),From="%TEMPLATENAME%",To="%PROJECTNAME%",bCaseSensitive=false)
\ No newline at end of file
+ReplacementsInFiles=(Extensions=("cpp","h","ini","cs"),From="%TEMPLATENAME%",To="%PROJECTNAME%",bCaseSensitive=false)
+
+; Which settings on the project settings page should be hidden for this template. Default: none
+; Allowed values are:
+; All - Hides the settings page entirely
+; Languages - Choose between C++ and Blueprint versions.
+; HardwareTarget - Choose between Mobile and Desktop hardware targets.
+; GraphicsPreset - Choose between Maximum and Scalable graphics.
+; StarterContent - Choose whether to add starter content to the project.
+; VR - Choose whether to enable Virtual Reality support for the project.
+;HiddenSettings=All
+HiddenSettings=Languages
+HiddenSettings=HardwareTarget
+HiddenSettings=GraphicsPreset 
+;HiddenSettings=StarterContent 
+HiddenSettings=VR 
\ No newline at end of file
diff --git a/Templates/TemplateCategories.ini b/Templates/TemplateCategories.ini
--- a/Templates/TemplateCategories.ini
+++ b/Templates/TemplateCategories.ini
@@ -12,7 +12,10 @@
 ; "Icon": A relative reference to the icon to display for the category - should be at minimum 128x128 to avoid being blurry.
 ; "IsMajorCategory": Is this a major, top-level category? Major categories are displayed as full rows, eg. the "Game" category.
 
-Categories=(Key="Game", LocalizedDisplayNames=((Language="en", Text="Games")), LocalizedDescriptions=((Language="en", Text="Starting off points giving you a few key classes, levels and examples to begin your game.")), IsMajorCategory=true)
-Categories=(Key="ME", LocalizedDisplayNames=((Language="en", Text="Media & Entertainment")), LocalizedDescriptions=((Language="en", Text="need text need text need text")), IsMajorCategory=true)
-Categories=(Key="AEC", LocalizedDisplayNames=((Language="en", Text="Architecture, Engineering & Construction")), LocalizedDescriptions=((Language="en", Text="need text need text need text")), IsMajorCategory=true)
-Categories=(Key="MFG", LocalizedDisplayNames=((Language="en", Text="Manufacturing, Automotive & Product Design")), LocalizedDescriptions=((Language="en", Text="need text need text need text")), IsMajorCategory=true)
\ No newline at end of file
+Categories=(Key="Games", LocalizedDisplayNames=((Language="en", Text="Games")), LocalizedDescriptions=((Language="en", Text="Start your game development journey with one of our key classes, levels, and examples.")), Icon="Media/Games.png", IsMajorCategory=true)
+
+Categories=(Key="ME", LocalizedDisplayNames=((Language="en", Text="Media and Entertainment")), LocalizedDescriptions=((Language="en", Text="Choose from templates and examples for nDisplay, VR Scouting, and virtual production workflows.")), Icon="Media/ME.png", IsMajorCategory=true)
+
+Categories=(Key="AEC", LocalizedDisplayNames=((Language="en", Text="Architecture, Engineering and Construction")), LocalizedDescriptions=((Language="en", Text="Select a starting point for multi-user design reviews, photorealistic architectural design visualizations, sunlight studies, or stylized renderings.")), Icon="Media/AEC.png", IsMajorCategory=true)
+
+Categories=(Key="MFG", LocalizedDisplayNames=((Language="en", Text="Manufacturing, Automotive and Product Design")), LocalizedDescriptions=((Language="en", Text="Find templates for multi-user design reviews, photobooth studio environments, and product configurators.")),  Icon="Media/MFG.png", IsMajorCategory=true)
\ No newline at end of file
