Fixed invalid references in templates to HiddenSettings value of "VR", which was renamed to "XR" in a prior code change.

diff --git a/Templates/TP_AEC_BlankBP/Config/TemplateDefs.ini b/Templates/TP_AEC_BlankBP/Config/TemplateDefs.ini
--- a/Templates/TP_AEC_BlankBP/Config/TemplateDefs.ini
+++ b/Templates/TP_AEC_BlankBP/Config/TemplateDefs.ini
@@ -75,7 +75,7 @@ ReplacementsInFiles=(Extensions=("cpp","h","ini","cs"), From="%TEMPLATENAME%", T
 ; HardwareTarget - Choose between Mobile and Desktop hardware targets.
 ; GraphicsPreset - Choose between Maximum and Scalable graphics.
 ; StarterContent - Choose whether to add starter content to the project.
-; VR - Choose whether to enable Virtual Reality support for the project.
+; XR - Choose whether to enable Mixed Reality support for the project.
 ;HiddenSettings=All
 HiddenSettings=Languages
 HiddenSettings=HardwareTarget
diff --git a/Templates/TP_Blank/Config/TemplateDefs.ini b/Templates/TP_Blank/Config/TemplateDefs.ini
--- a/Templates/TP_Blank/Config/TemplateDefs.ini
+++ b/Templates/TP_Blank/Config/TemplateDefs.ini
@@ -82,13 +82,13 @@ ReplacementsInFiles=(Extensions=("cpp","h","ini","cs"), From="%TEMPLATENAME%", T
 ; HardwareTarget - Choose between Mobile and Desktop hardware targets.
 ; GraphicsPreset - Choose between Maximum and Scalable graphics.
 ; StarterContent - Choose whether to add starter content to the project.
-; VR - Choose whether to enable Virtual Reality support for the project.
+; XR - Choose whether to enable Mixed Reality support for the project.
 ;HiddenSettings=All
 ;HiddenSettings=Languages
 ;HiddenSettings=HardwareTarget
 ;HiddenSettings=GraphicsPreset 
 ;HiddenSettings=StarterContent 
-;HiddenSettings=VR 
+;HiddenSettings=XR 
 
 ; The class types that this project uses.
 ; This is only for display in the template description, and does not affect project creation.
diff --git a/Templates/TP_ME_BlankBP/Config/TemplateDefs.ini b/Templates/TP_ME_BlankBP/Config/TemplateDefs.ini
--- a/Templates/TP_ME_BlankBP/Config/TemplateDefs.ini
+++ b/Templates/TP_ME_BlankBP/Config/TemplateDefs.ini
@@ -75,7 +75,7 @@ ReplacementsInFiles=(Extensions=("cpp","h","ini","cs"), From="%TEMPLATENAME%", T
 ; HardwareTarget - Choose between Mobile and Desktop hardware targets.
 ; GraphicsPreset - Choose between Maximum and Scalable graphics.
 ; StarterContent - Choose whether to add starter content to the project.
-; VR - Choose whether to enable Virtual Reality support for the project.
+; XR - Choose whether to enable Mixed Reality support for the project.
 ;HiddenSettings=All
 HiddenSettings=Languages
 HiddenSettings=HardwareTarget
diff --git a/Templates/TP_MFG_BlankBP/Config/TemplateDefs.ini b/Templates/TP_MFG_BlankBP/Config/TemplateDefs.ini
--- a/Templates/TP_MFG_BlankBP/Config/TemplateDefs.ini
+++ b/Templates/TP_MFG_BlankBP/Config/TemplateDefs.ini
@@ -75,7 +75,7 @@ ReplacementsInFiles=(Extensions=("cpp","h","ini","cs"), From="%TEMPLATENAME%", T
 ; HardwareTarget - Choose between Mobile and Desktop hardware targets.
 ; GraphicsPreset - Choose between Maximum and Scalable graphics.
 ; StarterContent - Choose whether to add starter content to the project.
-; VR - Choose whether to enable Virtual Reality support for the project.
+; XR - Choose whether to enable Mixed Reality support for the project.
 ;HiddenSettings=All
 HiddenSettings=Languages
 HiddenSettings=HardwareTarget
diff --git a/Templates/TP_PhotoStudioBP/Config/TemplateDefs.ini b/Templates/TP_PhotoStudioBP/Config/TemplateDefs.ini
--- a/Templates/TP_PhotoStudioBP/Config/TemplateDefs.ini
+++ b/Templates/TP_PhotoStudioBP/Config/TemplateDefs.ini
@@ -75,7 +75,7 @@ ReplacementsInFiles=(Extensions=("cpp","h","ini","cs"), From="%TEMPLATENAME%", T
 ; HardwareTarget - Choose between Mobile and Desktop hardware targets.
 ; GraphicsPreset - Choose between Maximum and Scalable graphics.
 ; StarterContent - Choose whether to add starter content to the project.
-; VR - Choose whether to enable Virtual Reality support for the project.
+; XR - Choose whether to enable Mixed Reality support for the project.
 ;HiddenSettings=All
 HiddenSettings=Languages
 HiddenSettings=HardwareTarget
diff --git a/Templates/TP_nDisplayBP/Config/TemplateDefs.ini b/Templates/TP_nDisplayBP/Config/TemplateDefs.ini
--- a/Templates/TP_nDisplayBP/Config/TemplateDefs.ini
+++ b/Templates/TP_nDisplayBP/Config/TemplateDefs.ini
@@ -46,10 +46,10 @@ ReplacementsInFiles=(Extensions=("cpp","h","ini","cs"),From="%TEMPLATENAME%",To=
 ; HardwareTarget - Choose between Mobile and Desktop hardware targets.
 ; GraphicsPreset - Choose between Maximum and Scalable graphics.
 ; StarterContent - Choose whether to add starter content to the project.
-; VR - Choose whether to enable Virtual Reality support for the project.
+; XR - Choose whether to enable Mixed Reality support for the project.
 ;HiddenSettings=All
 HiddenSettings=Languages
 HiddenSettings=HardwareTarget
 HiddenSettings=GraphicsPreset 
 ;HiddenSettings=StarterContent 
-HiddenSettings=VR 
\ No newline at end of file
+HiddenSettings=XR 
\ No newline at end of file
