Blank C++ template project now correctly shows the Compile button in the level editor toolbar.

diff --git a/Templates/TP_Blank/Config/TemplateDefs.ini b/Templates/TP_Blank/Config/TemplateDefs.ini
--- a/Templates/TP_Blank/Config/TemplateDefs.ini
+++ b/Templates/TP_Blank/Config/TemplateDefs.ini
@@ -84,12 +84,14 @@ ReplacementsInFiles=(Extensions=("cpp","h","ini","cs"), From="%TEMPLATENAME%", T
 ; Languages - Choose between C++ and Blueprint versions.
 ; HardwareTarget - Choose between Mobile and Desktop hardware targets.
 ; GraphicsPreset - Choose between Maximum and Scalable graphics.
+; Raytracing - Choose whether raytracing should be enabled for the project.
 ; StarterContent - Choose whether to add starter content to the project.
 ; XR - Choose whether to enable Mixed Reality support for the project.
 ;HiddenSettings=All
 ;HiddenSettings=Languages
 ;HiddenSettings=HardwareTarget
 ;HiddenSettings=GraphicsPreset 
+;HiddenSettings=Raytracing
 ;HiddenSettings=StarterContent 
 ;HiddenSettings=XR 
 
diff --git a/Templates/TP_Blank/TP_Blank.uproject b/Templates/TP_Blank/TP_Blank.uproject
--- a/Templates/TP_Blank/TP_Blank.uproject
+++ b/Templates/TP_Blank/TP_Blank.uproject
@@ -3,5 +3,11 @@
 	"EngineAssociation": "",
 	"Category": "",
 	"Description": "",
-	"Modules": []
+	"Modules": [
+		{
+			"Name": "TP_Blank",
+			"Type": "Runtime",
+			"LoadingPhase": "Default"
+		}
+	]
 }
\ No newline at end of file
