Merging //UE4/Dev-RenderPlat-Staging to Dev-Mobile (//UE4/Dev-Mobile)

diff --git a/Engine/Plugins/Runtime/AR/Apple/AppleARKit/Source/AppleARKitPoseTrackingLiveLink/Private/AppleARKitPoseTrackingLiveLinkImpl.cpp b/Engine/Plugins/Runtime/AR/Apple/AppleARKit/Source/AppleARKitPoseTrackingLiveLink/Private/AppleARKitPoseTrackingLiveLinkImpl.cpp
--- a/Engine/Plugins/Runtime/AR/Apple/AppleARKit/Source/AppleARKitPoseTrackingLiveLink/Private/AppleARKitPoseTrackingLiveLinkImpl.cpp
+++ b/Engine/Plugins/Runtime/AR/Apple/AppleARKit/Source/AppleARKitPoseTrackingLiveLink/Private/AppleARKitPoseTrackingLiveLinkImpl.cpp
@@ -14,7 +14,7 @@ FAppleARKitPoseTrackingLiveLink::FAppleARKitPoseTrackingLiveLink() :
 FAppleARKitPoseTrackingLiveLink::~FAppleARKitPoseTrackingLiveLink()
 {
 	// Should only be called durirng shutdown
-	check(GIsRequestingExit);
+	check(IsEngineExitRequested());
 }
 
 void FAppleARKitPoseTrackingLiveLink::Init()
diff --git a/Engine/Shaders/Private/PostProcessMaterialShaders.usf b/Engine/Shaders/Private/PostProcessMaterialShaders.usf
--- a/Engine/Shaders/Private/PostProcessMaterialShaders.usf
+++ b/Engine/Shaders/Private/PostProcessMaterialShaders.usf
@@ -199,6 +199,10 @@ void MainPS(
 	const float Alpha = 1.0f;
 #endif
 	OutColor = float4(GetMaterialEmissive(PixelMaterialInputs), Alpha );
+
+#if !MATERIALBLENDING_MODULATE && USE_PREEXPOSURE && !POST_PROCESS_MATERIAL_AFTER_TAA_UPSAMPLE
+	OutColor.xyz *= View.PreExposure;
+#endif
 } 
 
 void MainPS_VideoOverlay(
diff --git a/Engine/Source/Runtime/Apple/MetalRHI/Private/MetalPipeline.cpp b/Engine/Source/Runtime/Apple/MetalRHI/Private/MetalPipeline.cpp
--- a/Engine/Source/Runtime/Apple/MetalRHI/Private/MetalPipeline.cpp
+++ b/Engine/Source/Runtime/Apple/MetalRHI/Private/MetalPipeline.cpp
@@ -652,6 +652,12 @@ APPLE_PLATFORM_OBJECT_ALLOC_OVERRIDES(FMetalShaderPipeline)
 	{
 		MTLArgument* Arg = [Arguments objectAtIndex:i];
 		check(Arg);
+
+		if (!Arg.active)
+		{
+			continue;
+		}
+
 		switch(Arg.type)
 		{
 			case MTLArgumentTypeBuffer:
diff --git a/Engine/Source/Runtime/Engine/Private/Materials/HLSLMaterialTranslator.h b/Engine/Source/Runtime/Engine/Private/Materials/HLSLMaterialTranslator.h
--- a/Engine/Source/Runtime/Engine/Private/Materials/HLSLMaterialTranslator.h
+++ b/Engine/Source/Runtime/Engine/Private/Materials/HLSLMaterialTranslator.h
@@ -4876,11 +4876,20 @@ protected:
 
 		if (FeatureLevel >= ERHIFeatureLevel::SM5)
 		{
-			return AddCodeChunk(
+			int32 LookUp = AddCodeChunk(
 				MCT_Float4,
 				TEXT("SceneTextureLookup(%s, %d, %s)"),
 				*CoerceParameter(BufferUV, MCT_Float2), (int)SceneTextureId, bFiltered ? TEXT("true") : TEXT("false")
 				);
+
+			if (SceneTextureId == PPI_PostProcessInput0 && Material->GetMaterialDomain() == MD_PostProcess && Material->GetBlendableLocation() != BL_AfterTonemapping)
+			{
+				return AddInlinedCodeChunk(MCT_Float4, TEXT("(float4(View.OneOverPreExposure.xxx, 1) * %s)"), *CoerceParameter(LookUp, MCT_Float4));
+			}
+			else
+			{
+				return LookUp;
+			}
 		}
 		else // mobile
 		{
diff --git a/Engine/Source/Runtime/RenderCore/Private/ClearQuad.cpp b/Engine/Source/Runtime/RenderCore/Private/ClearQuad.cpp
deleted file mode 100644
diff --git a/Engine/Source/Runtime/RenderCore/Private/ClearReplacementShaders.cpp b/Engine/Source/Runtime/RenderCore/Private/ClearReplacementShaders.cpp
deleted file mode 100644
diff --git a/Engine/Source/Runtime/RenderCore/Private/LongGPUTask.cpp b/Engine/Source/Runtime/RenderCore/Private/LongGPUTask.cpp
deleted file mode 100644
diff --git a/Engine/Source/Runtime/RenderCore/Private/MediaShaders.cpp b/Engine/Source/Runtime/RenderCore/Private/MediaShaders.cpp
deleted file mode 100644
diff --git a/Engine/Source/Runtime/RenderCore/Private/OculusShaders.cpp b/Engine/Source/Runtime/RenderCore/Private/OculusShaders.cpp
deleted file mode 100644
diff --git a/Engine/Source/Runtime/RenderCore/Private/OneColorShader.cpp b/Engine/Source/Runtime/RenderCore/Private/OneColorShader.cpp
deleted file mode 100644
diff --git a/Engine/Source/Runtime/RenderCore/Private/RGBAToYUV420Shader.cpp b/Engine/Source/Runtime/RenderCore/Private/RGBAToYUV420Shader.cpp
deleted file mode 100644
diff --git a/Engine/Source/Runtime/RenderCore/Private/ResolveShader.cpp b/Engine/Source/Runtime/RenderCore/Private/ResolveShader.cpp
deleted file mode 100644
diff --git a/Engine/Source/Runtime/RenderCore/Public/BuiltInRayTracingShaders.h b/Engine/Source/Runtime/RenderCore/Public/BuiltInRayTracingShaders.h
deleted file mode 100644
diff --git a/Engine/Source/Runtime/RenderCore/Public/ClearQuad.h b/Engine/Source/Runtime/RenderCore/Public/ClearQuad.h
deleted file mode 100644
diff --git a/Engine/Source/Runtime/RenderCore/Public/ClearReplacementShaders.h b/Engine/Source/Runtime/RenderCore/Public/ClearReplacementShaders.h
deleted file mode 100644
diff --git a/Engine/Source/Runtime/RenderCore/Public/LongGPUTask.h b/Engine/Source/Runtime/RenderCore/Public/LongGPUTask.h
deleted file mode 100644
diff --git a/Engine/Source/Runtime/RenderCore/Public/MediaShaders.h b/Engine/Source/Runtime/RenderCore/Public/MediaShaders.h
deleted file mode 100644
diff --git a/Engine/Source/Runtime/RenderCore/Public/OculusShaders.h b/Engine/Source/Runtime/RenderCore/Public/OculusShaders.h
deleted file mode 100644
diff --git a/Engine/Source/Runtime/RenderCore/Public/OneColorShader.h b/Engine/Source/Runtime/RenderCore/Public/OneColorShader.h
deleted file mode 100644
diff --git a/Engine/Source/Runtime/RenderCore/Public/RGBAToYUV420Shader.h b/Engine/Source/Runtime/RenderCore/Public/RGBAToYUV420Shader.h
deleted file mode 100644
diff --git a/Engine/Source/Runtime/RenderCore/Public/ResolveShader.h b/Engine/Source/Runtime/RenderCore/Public/ResolveShader.h
deleted file mode 100644
diff --git a/Templates/TP_Blank/Config/TemplateDefs.ini b/Templates/TP_Blank/Config/TemplateDefs.ini
new file mode 100644
--- /dev/null
+++ b/Templates/TP_Blank/Config/TemplateDefs.ini
@@ -0,0 +1,106 @@
+; This file specifies how a template will be treated when creating a project from it.
+; It includes rules for ignoring files, replacing files and changing the contents of files.
+; This file is annotated to help with creating your own templates.
+
+[/Script/GameProjectGeneration.TemplateProjectDefs]
+
+; The localized display names for this template. Each entry specifies a Language code and Unicode text to display in the New Project dialog.
+LocalizedDisplayNames=(Language="en", Text="Blank")
+
+; The localized descriptions for this template, which is displayed when the template is selected.
+LocalizedDescriptions=(Language="en", Text="A clean empty project with no code.")
+
+; The categories that this template should appear in. The values in this list refer to the Key value given in TemplateCategories.ini.
+; Note: Names should not be quoted. You can add more categories as shown with the commented out line.
+;Categories=Game
+;Categories=ME
+;Categories=AEC
+;Categories=MFG
+
+; The sort key to use for the template to bypass the default alphabetical sorting.
+; Omit this if you just want to use default sorting.
+SortKey="_1"
+
+; A list of folders to avoid copying when creating a project from this template. 
+; These should contain files only of concern to the template itself and not desired in the user's project.
+FoldersToIgnore=Media
+
+; The files to ignore copying, relative to the template root folder. 
+; %TEMPLATENAME% will be replaced by the template's folder name when resolving these. 
+; Filename comparisons are case-insensitive. 
+; For example, if this template is "TP_Blank" and the user creates "MyProject", then "TP_Blank.uproject" will be copied to "MyProject.uproject".
+FilesToIgnore="%TEMPLATENAME%.uproject"
+FilesToIgnore="%TEMPLATENAME%.png"
+FilesToIgnore="Config/TemplateDefs.ini"
+FilesToIgnore="Config/config.ini"
+FilesToIgnore="%TEMPLATENAME%.opensdf"
+FilesToIgnore="%TEMPLATENAME%.sdf"
+FilesToIgnore="%TEMPLATENAME%.v11.suo"
+FilesToIgnore="%TEMPLATENAME%.v12.suo"
+FilesToIgnore="%TEMPLATENAME%.sln"
+FilesToIgnore="Manifest.json"
+FilesToIgnore="contents.txt"
+
+; NOTE: For the purposes of the below examples, this template is considered to be named "TP_Blank", and the user has created a project called "MyProject".
+
+; Folders that will be renamed when creating a project from this template.
+; %PROJECTNAME% will be replaced by the project name specified by the user during creation when resolving these rules.
+; For example,  then "Source/TP_Blank" will be copied to "Source/MyProject".
+FolderRenames=(From="Source/%TEMPLATENAME%", To="Source/%PROJECTNAME%")
+FolderRenames=(From="Source/%TEMPLATENAME%Editor", To="Source/%PROJECTNAME%Editor")
+
+; Rules for renaming files when creating a project from this template. 
+; %TEMPLATENAME_UPPERCASE% and %PROJECTNAME_UPPERCASE% are the template and project names in upper case, eg. for "TP_Blank", its upper case representation is "TP_BLANK".
+; %TEMPLATENAME_LOWERCASE% and %PROJECTNAME_LOWERCASE% are the template and project names in lower case, eg. for "TP_Blank", its lower case representation is "tp_blank".
+; A rule does not have to include any of these modifiers. 
+; Extensions is a list of files to consider for these replacements. Eg. ("cpp", "h") will try to apply a rule only to files like "SomeFile.cpp" and "SomeFile.h".
+; bCaseSensitive specifies whether this replacement should be case-sensitive.
+
+; Here are some default rules for renaming files:
+
+; This rule states that a file called "TP_BLANK.cpp" will be renamed to "MYPROJECT.cpp".
+FilenameReplacements=(Extensions=("cpp","h","ini","cs"), From="%TEMPLATENAME_UPPERCASE%", To="%PROJECTNAME_UPPERCASE%", bCaseSensitive=true)
+
+; This rule states that a file called "tp_blank.ini" will be renamed to "myproject.ini"
+FilenameReplacements=(Extensions=("cpp","h","ini","cs"), From="%TEMPLATENAME_LOWERCASE%", To="%PROJECTNAME_LOWERCASE%", bCaseSensitive=true)
+
+; And this rule catches any irregularly-cased files, so that "Tp_bLANK.ini" will be renamed to "MyProject.ini".
+FilenameReplacements=(Extensions=("cpp","h","ini","cs"), From="%TEMPLATENAME%",To="%PROJECTNAME%",bCaseSensitive=false)
+
+; Rules for renaming strings in the contents of files. 
+; The same rules as above apply, but any usages of strings that match these rules inside the file will be changed. This includes all substrings!
+; For example, "TP_Blank" will be replaced by "MyProject", but "TP_Blank_Name" will also be replaced by "MyProject_Name".
+ReplacementsInFiles=(Extensions=("cpp","h","ini","cs"), From="%TEMPLATENAME_UPPERCASE%", To="%PROJECTNAME_UPPERCASE%", bCaseSensitive=true)
+ReplacementsInFiles=(Extensions=("cpp","h","ini","cs"), From="%TEMPLATENAME_LOWERCASE%", To="%PROJECTNAME_LOWERCASE%", bCaseSensitive=true)
+ReplacementsInFiles=(Extensions=("cpp","h","ini","cs"), From="%TEMPLATENAME%", To="%PROJECTNAME%", bCaseSensitive=false)
+
+; The following settings are not enabled for this project, but are included for completeness:
+
+; Whether or not this template can be used to create a project. Default: true
+; If this is set to false, then the template becomes a Feature Pack that can only be referenced from other templates, or added to a project from the Content Browser.
+;bAllowProjectCreation=false
+
+; Whether or not to skip the project settings page entirely for this template. Default: false
+;bSkipProjectSettings=true
+
+; Which settings on the project settings page should be hidden for this template. Default: none
+; Allowed values are:
+; Languages - Choose between C++ and Blueprint versions.
+; HardwareTarget - Choose between Mobile and Desktop hardware targets.
+; GraphicsPreset - Choose between Maximum and Scalable graphics.
+; StarterContent - Choose whether to add starter content to the project.
+; VR - Choose whether to enable Virtual Reality support for the project.
+;HiddenSettings=Languages
+;HiddenSettings=HardwareTarget
+
+
+; The class types that this project uses.
+; This is only for display in the template description, and does not affect project creation.
+;ClassTypes=GameMode, Character, Actor
+
+; The asset types that this project uses.
+; This is only for display in the template description, and does not affect project creation.
+;AssetTypes=Skeleton, Static Mesh, Texture
+
+; Shared feature packs. The files in these packs listed in these structures marked as 'additionalfiles' will be copied on project generation.
+;SharedContentPacks=(MountName="2DSideScroller",DetailLevels=("Standard"))
diff --git a/Templates/TP_Blank/Source/TP_Blank.Target.cs b/Templates/TP_Blank/Source/TP_Blank.Target.cs
new file mode 100644
--- /dev/null
+++ b/Templates/TP_Blank/Source/TP_Blank.Target.cs
@@ -0,0 +1,14 @@
+// Copyright 1998-2019 Epic Games, Inc. All Rights Reserved.
+
+using UnrealBuildTool;
+using System.Collections.Generic;
+
+public class TP_BlankTarget : TargetRules
+{
+	public TP_BlankTarget( TargetInfo Target) : base(Target)
+	{
+		Type = TargetType.Game;
+
+		ExtraModuleNames.AddRange( new string[] { "TP_Blank" } );
+	}
+}
diff --git a/Templates/TP_Blank/Source/TP_Blank/TP_Blank.Build.cs b/Templates/TP_Blank/Source/TP_Blank/TP_Blank.Build.cs
new file mode 100644
--- /dev/null
+++ b/Templates/TP_Blank/Source/TP_Blank/TP_Blank.Build.cs
@@ -0,0 +1,23 @@
+// Copyright 1998-2019 Epic Games, Inc. All Rights Reserved.
+
+using UnrealBuildTool;
+
+public class TP_Blank : ModuleRules
+{
+	public TP_Blank(ReadOnlyTargetRules Target) : base(Target)
+	{
+		PCHUsage = PCHUsageMode.UseExplicitOrSharedPCHs;
+	
+		PublicDependencyModuleNames.AddRange(new string[] { "Core", "CoreUObject", "Engine", "InputCore" });
+
+		PrivateDependencyModuleNames.AddRange(new string[] {  });
+
+		// Uncomment if you are using Slate UI
+		// PrivateDependencyModuleNames.AddRange(new string[] { "Slate", "SlateCore" });
+		
+		// Uncomment if you are using online features
+		// PrivateDependencyModuleNames.Add("OnlineSubsystem");
+
+		// To include OnlineSubsystemSteam, add it to the plugins section in your uproject file with the Enabled attribute set to true
+	}
+}
diff --git a/Templates/TP_Blank/Source/TP_Blank/TP_Blank.cpp b/Templates/TP_Blank/Source/TP_Blank/TP_Blank.cpp
new file mode 100644
--- /dev/null
+++ b/Templates/TP_Blank/Source/TP_Blank/TP_Blank.cpp
@@ -0,0 +1,6 @@
+// Copyright 1998-2019 Epic Games, Inc. All Rights Reserved.
+
+#include "TP_Blank.h"
+#include "Modules/ModuleManager.h"
+
+IMPLEMENT_PRIMARY_GAME_MODULE( FDefaultGameModuleImpl, TP_Blank, "TP_Blank" );
diff --git a/Templates/TP_Blank/Source/TP_Blank/TP_Blank.h b/Templates/TP_Blank/Source/TP_Blank/TP_Blank.h
new file mode 100644
--- /dev/null
+++ b/Templates/TP_Blank/Source/TP_Blank/TP_Blank.h
@@ -0,0 +1,6 @@
+// Copyright 1998-2019 Epic Games, Inc. All Rights Reserved.
+
+#pragma once
+
+#include "CoreMinimal.h"
+
diff --git a/Templates/TP_Blank/Source/TP_Blank/TP_BlankGameModeBase.cpp b/Templates/TP_Blank/Source/TP_Blank/TP_BlankGameModeBase.cpp
new file mode 100644
--- /dev/null
+++ b/Templates/TP_Blank/Source/TP_Blank/TP_BlankGameModeBase.cpp
@@ -0,0 +1,5 @@
+// Copyright 1998-2019 Epic Games, Inc. All Rights Reserved.
+
+
+#include "TP_BlankGameModeBase.h"
+
diff --git a/Templates/TP_Blank/Source/TP_Blank/TP_BlankGameModeBase.h b/Templates/TP_Blank/Source/TP_Blank/TP_BlankGameModeBase.h
new file mode 100644
--- /dev/null
+++ b/Templates/TP_Blank/Source/TP_Blank/TP_BlankGameModeBase.h
@@ -0,0 +1,17 @@
+// Copyright 1998-2019 Epic Games, Inc. All Rights Reserved.
+
+#pragma once
+
+#include "CoreMinimal.h"
+#include "GameFramework/GameModeBase.h"
+#include "TP_BlankGameModeBase.generated.h"
+
+/**
+ * 
+ */
+UCLASS()
+class TP_BLANK_API ATP_BlankGameModeBase : public AGameModeBase
+{
+	GENERATED_BODY()
+	
+};
diff --git a/Templates/TP_Blank/Source/TP_BlankEditor.Target.cs b/Templates/TP_Blank/Source/TP_BlankEditor.Target.cs
new file mode 100644
--- /dev/null
+++ b/Templates/TP_Blank/Source/TP_BlankEditor.Target.cs
@@ -0,0 +1,14 @@
+// Copyright 1998-2019 Epic Games, Inc. All Rights Reserved.
+
+using UnrealBuildTool;
+using System.Collections.Generic;
+
+public class TP_BlankEditorTarget : TargetRules
+{
+	public TP_BlankEditorTarget( TargetInfo Target) : base(Target)
+	{
+		Type = TargetType.Editor;
+
+		ExtraModuleNames.AddRange( new string[] { "TP_Blank" } );
+	}
+}
diff --git a/Templates/TP_Blank/TP_Blank.uproject b/Templates/TP_Blank/TP_Blank.uproject
new file mode 100644
--- /dev/null
+++ b/Templates/TP_Blank/TP_Blank.uproject
@@ -0,0 +1,7 @@
+{
+	"FileVersion": 3,
+	"EngineAssociation": "",
+	"Category": "",
+	"Description": "",
+	"Modules": []
+}
\ No newline at end of file
diff --git a/Templates/TP_Blank/contents.txt b/Templates/TP_Blank/contents.txt
new file mode 100644
--- /dev/null
+++ b/Templates/TP_Blank/contents.txt
@@ -0,0 +1,4 @@
+../../../Templates/TP_Blank/manifest.json
+../../../Templates/TP_Blank/Config/config.ini
+../../../Templates/TP_Blank/Media/TP_Blank.png
+../../../Templates/TP_Blank/Media/TP_Blank_Preview.png
diff --git a/Templates/TP_Blank/manifest.json b/Templates/TP_Blank/manifest.json
new file mode 100644
--- /dev/null
+++ b/Templates/TP_Blank/manifest.json
@@ -0,0 +1,39 @@
+{
+	"Version": 1,
+	"Name":
+	[
+		{
+			"Language": "en",
+			"Text": "Blank"
+		}
+	],
+	"Description":
+	[
+		{
+			"Language": "en",
+			"Text": "An empty project with some basic game framework code classes created."
+		}
+	],
+	"AssetTypes": [],	
+	"SearchTags": [],
+	"ClassTypes": "",
+	"Category": "",
+	"FocusAsset": "",
+	"Thumbnail": "TP_Blank.png",
+	"Screenshots":
+	[
+		"TP_Blank_Preview.png"
+	],
+	"SortKey": "_1",
+	"AdditionalFeaturePacks":
+	[
+		{
+			"MountName": "Geometry",
+			"DetailLevels":
+			[
+				"Standard",
+				"High"
+			]
+		}
+	]
+}
diff --git a/Templates/TP_BlankBP/Config/TemplateDefs.ini b/Templates/TP_BlankBP/Config/TemplateDefs.ini
new file mode 100644
--- /dev/null
+++ b/Templates/TP_BlankBP/Config/TemplateDefs.ini
@@ -0,0 +1,37 @@
+[/Script/GameProjectGeneration.TemplateProjectDefs]
+LocalizedDisplayNames=(Language="en", Text="Blank")
+LocalizedDescriptions=(Language="en", Text="A clean empty project with no code.")
+
+SortKey="_1"
+
+;Categories=Game
+;Categories=ME
+;Categories=AEC
+;Categories=MFG
+
+FoldersToIgnore=Media
+
+FilesToIgnore="%TEMPLATENAME%.uproject"
+FilesToIgnore="%TEMPLATENAME%.png"
+FilesToIgnore="Config/TemplateDefs.ini"
+FilesToIgnore="Config/config.ini"
+FilesToIgnore="%TEMPLATENAME%.opensdf"
+FilesToIgnore="%TEMPLATENAME%.sdf"
+FilesToIgnore="%TEMPLATENAME%.v11.suo"
+FilesToIgnore="%TEMPLATENAME%.v12.suo"
+FilesToIgnore="%TEMPLATENAME%.sln"
+FilesToIgnore="Manifest.json"
+FilesToIgnore="contents.txt"
+
+FolderRenames=(From="Source/%TEMPLATENAME%", To="Source/%PROJECTNAME%")
+FolderRenames=(From="Source/%TEMPLATENAME%Editor", To="Source/%PROJECTNAME%Editor")
+
+FilenameReplacements=(Extensions=("cpp","h","ini","cs"), From="%TEMPLATENAME_UPPERCASE%", To="%PROJECTNAME_UPPERCASE%", bCaseSensitive=true)
+FilenameReplacements=(Extensions=("cpp","h","ini","cs"), From="%TEMPLATENAME_LOWERCASE%", To="%PROJECTNAME_LOWERCASE%", bCaseSensitive=true)
+FilenameReplacements=(Extensions=("cpp","h","ini","cs"), From="%TEMPLATENAME%", To="%PROJECTNAME%", bCaseSensitive=false)
+
+ReplacementsInFiles=(Extensions=("cpp","h","ini","cs"), From="%TEMPLATENAME_UPPERCASE%", To="%PROJECTNAME_UPPERCASE%", bCaseSensitive=true)
+ReplacementsInFiles=(Extensions=("cpp","h","ini","cs"), From="%TEMPLATENAME_LOWERCASE%", To="%PROJECTNAME_LOWERCASE%", bCaseSensitive=true)
+ReplacementsInFiles=(Extensions=("cpp","h","ini","cs"), From="%TEMPLATENAME%", To="%PROJECTNAME%", bCaseSensitive=false)
+
+SharedContentPacks=(MountName="Geometry",DetailLevels=("Standard","High"))
\ No newline at end of file
diff --git a/Templates/TP_BlankBP/TP_BlankBP.uproject b/Templates/TP_BlankBP/TP_BlankBP.uproject
new file mode 100644
--- /dev/null
+++ b/Templates/TP_BlankBP/TP_BlankBP.uproject
@@ -0,0 +1,7 @@
+{
+	"FileVersion": 3,
+	"EngineAssociation": "",
+	"Category": "",
+	"Description": "",
+	"Modules": []
+}
\ No newline at end of file
diff --git a/Templates/TP_BlankBP/contents.txt b/Templates/TP_BlankBP/contents.txt
new file mode 100644
--- /dev/null
+++ b/Templates/TP_BlankBP/contents.txt
@@ -0,0 +1,4 @@
+../../../Templates/TP_BlankBP/manifest.json
+../../../Templates/TP_BlankBP/Config/config.ini
+../../../Templates/TP_BlankBP/Media/TP_BlankBP.png
+../../../Templates/TP_BlankBP/Media/TP_BlankBP_Preview.png
diff --git a/Templates/TP_BlankBP/manifest.json b/Templates/TP_BlankBP/manifest.json
new file mode 100644
--- /dev/null
+++ b/Templates/TP_BlankBP/manifest.json
@@ -0,0 +1,39 @@
+{
+	"Version": 1,
+	"Name":
+	[
+		{
+			"Language": "en",
+			"Text": "Blank"
+		}
+	],
+	"Description":
+	[
+		{
+			"Language": "en",
+			"Text": "A clean empty project with no code."
+		}
+	],
+	"AssetTypes": [],
+	"SearchTags": [],
+	"ClassTypes": "",
+	"Category": "",
+	"FocusAsset": "",
+	"Thumbnail": "TP_BlankBP.png",
+	"Screenshots":
+	[
+		"TP_BlankBP_Preview.png"
+	],
+	"SortKey": "_1",
+	"AdditionalFeaturePacks":
+	[
+		{
+			"MountName": "Geometry",
+			"DetailLevels":
+			[
+				"Standard",
+				"High"
+			]
+		}
+	]
+}
diff --git a/Templates/TemplateCategories.ini b/Templates/TemplateCategories.ini
new file mode 100644
--- /dev/null
+++ b/Templates/TemplateCategories.ini
@@ -0,0 +1,18 @@
+; This file lists all the template categories that are exposed from this folder.
+; Add template categories here to have them show up in the Project Browser.
+; Additional TemplateDefs.ini files can be created in your plugin root folder, at the same level as a Templates/ folder.
+; Categories will contain project templates that include the category Key in their Categories property in their Config/TemplateDefs.ini.
+
+[/Script/GameProjectGeneration.TemplateCategories]
+
+; Each "Categories" entry adds a separate category with these fields:
+; "Key": The name used to reference the category from templates. Should be short and easy.
+; "LocalizedDisplayNames": A list of pairs of "Language" codes and "Text" to display as the name of the category in the UI (supports Unicode).
+; "LocalizedDescriptions": As with display names, this is the description shown next to each category.
+; "Icon": A relative reference to the icon to display for the category - should be at minimum 128x128 to avoid being blurry.
+; "IsMajorCategory": Is this a major, top-level category? Major categories are displayed as full rows, eg. the "Game" category.
+
+Categories=(Key="Game", LocalizedDisplayNames=((Language="en", Text="Games")), LocalizedDescriptions=((Language="en", Text="Starting off points giving you a few key classes, levels and examples to begin your game.")), IsMajorCategory=true)
+Categories=(Key="ME", LocalizedDisplayNames=((Language="en", Text="Media & Entertainment")), LocalizedDescriptions=((Language="en", Text="need text need text need text")), IsMajorCategory=true)
+Categories=(Key="AEC", LocalizedDisplayNames=((Language="en", Text="Architecture, Engineering & Construction")), LocalizedDescriptions=((Language="en", Text="need text need text need text")), IsMajorCategory=true)
+Categories=(Key="MFG", LocalizedDisplayNames=((Language="en", Text="Manufacturing, Automotive & Product Design")), LocalizedDescriptions=((Language="en", Text="need text need text need text")), IsMajorCategory=true)
\ No newline at end of file
