Adding Chinese translation to this file #jira UELOC-4711 #fyi jamie.dale alexis.zhang michelle.zhao yali.he

diff --git a/Templates/TP_Blank/Config/TemplateDefs.ini b/Templates/TP_Blank/Config/TemplateDefs.ini
--- a/Templates/TP_Blank/Config/TemplateDefs.ini
+++ b/Templates/TP_Blank/Config/TemplateDefs.ini
@@ -1,4 +1,4 @@
-; This file specifies how a template will be treated when creating a project from it.
+﻿; This file specifies how a template will be treated when creating a project from it.
 ; It includes rules for ignoring files, replacing files and changing the contents of files.
 ; This file is annotated to help with creating your own templates.
 
@@ -9,9 +9,11 @@ Categories=Games
 
 ; The localized display names for this template. Each entry specifies a Language code and Unicode text to display in the New Project dialog.
 LocalizedDisplayNames=(Language="en", Text="Blank")
+LocalizedDisplayNames=(Language="zh-Hans", Text="空白")
 
 ; The localized descriptions for this template, which is displayed when the template is selected.
 LocalizedDescriptions=(Language="en", Text="A clean empty project with no code.")
+LocalizedDescriptions=(Language="zh-Hans", Text="不含任何代码的空白项目。")
 
 ; The sort key to use for the template to bypass the default alphabetical sorting.
 SortKey="_1"
