Adding Japanese translation to this file #jira ueloc-4712 #fyi yali.he, michelle.zhao

diff --git a/Templates/TP_Blank/Config/TemplateDefs.ini b/Templates/TP_Blank/Config/TemplateDefs.ini
--- a/Templates/TP_Blank/Config/TemplateDefs.ini
+++ b/Templates/TP_Blank/Config/TemplateDefs.ini
@@ -14,6 +14,7 @@ LocalizedDisplayNames=(Language="zh-Hans", Text="空白")
 ; The localized descriptions for this template, which is displayed when the template is selected.
 LocalizedDescriptions=(Language="en", Text="A clean empty project with no code.")
 LocalizedDescriptions=(Language="zh-Hans", Text="不含任何代码的空白项目。")
+LocalizedDescriptions=(Language="ja", Text="まったく空のプロジェクトです。コードは入っていません。")
 
 ; The sort key to use for the template to bypass the default alphabetical sorting.
 SortKey="_1"
